Removed unused reference to a developer module in my runtime module.

diff --git a/Engine/Source/Runtime/BlueprintRuntime/BlueprintRuntime.Build.cs b/Engine/Source/Runtime/BlueprintRuntime/BlueprintRuntime.Build.cs
--- a/Engine/Source/Runtime/BlueprintRuntime/BlueprintRuntime.Build.cs
+++ b/Engine/Source/Runtime/BlueprintRuntime/BlueprintRuntime.Build.cs
@@ -17,7 +17,6 @@ namespace UnrealBuildTool.Rules
 				{
 					"Core",
 				    "CoreUObject",
-                    "Settings",
 				}
 			);
 		}
diff --git a/Engine/Source/Runtime/BlueprintRuntime/Private/BlueprintRuntime.cpp b/Engine/Source/Runtime/BlueprintRuntime/Private/BlueprintRuntime.cpp
--- a/Engine/Source/Runtime/BlueprintRuntime/Private/BlueprintRuntime.cpp
+++ b/Engine/Source/Runtime/BlueprintRuntime/Private/BlueprintRuntime.cpp
@@ -1,7 +1,6 @@
 // Copyright 1998-2016 Epic Games, Inc. All Rights Reserved.
 
 #include "BlueprintRuntimePrivatePCH.h"
-#include "ISettingsModule.h"
 #include "BlueprintRuntimeSettings.h"
 #include "BlueprintSupport.h"
 
