Fix uninitialized member

diff --git a/Engine/Source/Runtime/BlueprintRuntime/Public/BlueprintRuntimeSettings.h b/Engine/Source/Runtime/BlueprintRuntime/Public/BlueprintRuntimeSettings.h
--- a/Engine/Source/Runtime/BlueprintRuntime/Public/BlueprintRuntimeSettings.h
+++ b/Engine/Source/Runtime/BlueprintRuntime/Public/BlueprintRuntimeSettings.h
@@ -29,7 +29,7 @@ struct FBlueprintWarningSettings
 	FText WarningDescription;
 
 	UPROPERTY()
-	EBlueprintWarningBehavior WarningBehavior;
+	EBlueprintWarningBehavior WarningBehavior = EBlueprintWarningBehavior::Warn;
 };
 
 /**
